Moved Gauntlet Plugin out of NFL and into experimental folder

#tests Compiled ElementalDemo

diff --git a/Engine/Plugins/Experimental/Gauntlet/Gauntlet.uplugin b/Engine/Plugins/Experimental/Gauntlet/Gauntlet.uplugin
new file mode 100644
--- /dev/null
+++ b/Engine/Plugins/Experimental/Gauntlet/Gauntlet.uplugin
@@ -0,0 +1,25 @@
+{
+	"FileVersion" : 3,
+	"Version" : 1,
+	"VersionName" : "0.1",
+	"FriendlyName" : "Gauntlet",
+	"Description" : "Provides a helper class for creating and managing tests in your game",
+	"Category" : "Automation",
+	"CreatedBy" : "Epic Games",
+	"CreatedByURL" : "",
+	"DocsURL" : "",
+	"MarketplaceURL" : "",
+	"SupportURL" : "",
+	"EnabledByDefault" : false,
+	"CanContainContent" : false,
+	"IsBetaVersion" : true,
+	"Installed" : false,
+	"Modules" :
+	[
+		{
+			"Name" : "Gauntlet",
+			"Type" : "Developer",
+			"LoadingPhase": "PreDefault"
+		}
+	]
+}
\ No newline at end of file
diff --git a/Engine/Plugins/Experimental/Gauntlet/Source/Gauntlet/Gauntlet.Build.cs b/Engine/Plugins/Experimental/Gauntlet/Source/Gauntlet/Gauntlet.Build.cs
new file mode 100644
--- /dev/null
+++ b/Engine/Plugins/Experimental/Gauntlet/Source/Gauntlet/Gauntlet.Build.cs
@@ -0,0 +1,23 @@
+// Copyright 1998-2018 Epic Games, Inc. All Rights Reserved.
+
+namespace UnrealBuildTool.Rules
+{
+	public class Gauntlet : ModuleRules
+	{
+		public Gauntlet(ReadOnlyTargetRules Target) : base(Target)
+		{
+			PublicDependencyModuleNames.AddRange(
+				new string[] {
+					"Core",
+					"CoreUObject",
+					"Engine",
+				});
+
+			PublicIncludePaths.AddRange(
+				new string[]
+				{
+				}
+			);
+		}
+	}
+}
\ No newline at end of file
diff --git a/Engine/Plugins/Experimental/Gauntlet/Source/Gauntlet/Private/GauntletModule.cpp b/Engine/Plugins/Experimental/Gauntlet/Source/Gauntlet/Private/GauntletModule.cpp
new file mode 100644
--- /dev/null
+++ b/Engine/Plugins/Experimental/Gauntlet/Source/Gauntlet/Private/GauntletModule.cpp
@@ -0,0 +1,341 @@
+// Copyright 1998-2018 Epic Games, Inc. All Rights Reserved.
+#include "GauntletModule.h"
+#include "Engine/Engine.h"
+#include "UObject/Package.h"
+#include "CoreMinimal.h"
+#include "Misc/CommandLine.h"
+#include "GauntletTestController.h"
+#include "UObject/UObjectIterator.h"
+#include "Containers/Ticker.h"
+#include "GameFramework/GameStateBase.h"
+#include "Misc/CoreDelegates.h"
+#include "UnrealClient.h"
+
+
+
+class GAUNTLET_API FGauntletModuleImpl
+	: public FGauntletModule
+	, public FGCObject
+
+{
+public:
+	/** IModuleInterface interface */
+	void StartupModule() override;
+	void ShutdownModule() override;
+
+	/** GauntletModule interface */
+	void	SetGameStateToTestStateMapping(const TMap<UClass*, FName>& Mapping) override;
+	void	SetWorldToTestStateMapping(const TMap<FString, FName>& Mapping) override;
+	void	BroadcastStateChange(FName NewState) override;
+	FName	GetCurrentState() const override { return CurrentState; }
+	double	GetTimeInCurrentState() const override { return TimeInCurrentState; }
+
+	
+	FString			GetCurrentMap() const { return CurrentMap; }
+
+	void			SetScreenshotPeriod(float Period)
+	{
+		ScreenshotPeriod = Period;
+	}
+
+	UGauntletTestController*	GetTestController(UClass* ControllerClass) override;
+
+protected:
+
+	/**
+	* FGCObject interface - add references to our UObject-inherited controllers to prevent
+	* them being collected
+	*/
+	void		AddReferencedObjects(FReferenceCollector& Collector) override;
+
+	/** Handler for PostMapChange delegate */
+	void		InnerPostMapChange(UWorld* World);
+
+	/** Handler for PreMapChange delegate */
+	void		InnerPreMapChange(const FString& MapName);
+
+	/** Tick callback */
+	void		InnerTick(const float TimeDelta);
+
+	/** Performs actual initialization */
+	void		OnPostEngineInit();
+
+	/** Loads all controllers based on command line arts */
+	void		LoadControllers();
+
+	void		TakeScreenshot();
+
+private:
+
+	/** Handle to our tick callback */
+	FDelegateHandle					TickHandle;
+
+	/** Currently active controllers */
+	TArray<UGauntletTestController*>	Controllers;
+
+	/** True if a state has been set */
+	bool							StateSet;
+
+	FName							CurrentState;
+	double							TimeInCurrentState;
+	FString							CurrentMap;
+	TMap<UClass*, FName>			GameStateStateMap;
+	TMap<FString, FName>			MapGameStateMap;
+	UClass*							CurrentGameStateClass;
+
+	double							LastScreenshotTime;
+	float							ScreenshotPeriod;
+};
+
+
+IMPLEMENT_MODULE(FGauntletModuleImpl, Gauntlet)
+
+DEFINE_LOG_CATEGORY(LogGauntlet);
+
+FName FGauntletStates::Initialized = TEXT("Gauntlet_Initialized");
+
+
+void FGauntletModuleImpl::StartupModule()
+{
+	StateSet = false;
+	CurrentGameStateClass = nullptr;
+	TimeInCurrentState = 0;
+	LastScreenshotTime = 0.0f;
+	ScreenshotPeriod = 0.0f;
+
+	if (IsRunningGame() || IsRunningDedicatedServer())
+	{
+		UE_LOG(LogGauntlet, Display, TEXT("Gauntlet Initialized"));
+		FCoreDelegates::OnPostEngineInit.AddRaw(this, &FGauntletModuleImpl::OnPostEngineInit);	
+	}
+}
+
+void FGauntletModuleImpl::OnPostEngineInit()
+{
+	FCoreUObjectDelegates::PostLoadMapWithWorld.AddRaw(this, &FGauntletModuleImpl::InnerPostMapChange);
+	FCoreUObjectDelegates::PreLoadMap.AddRaw(this, &FGauntletModuleImpl::InnerPreMapChange);
+
+	FParse::Value(FCommandLine::Get(), TEXT("gauntlet.screenshotperiod="), ScreenshotPeriod);
+
+	LoadControllers();
+
+	const float kTickRate = 1.0f;
+
+	FTicker::GetCoreTicker().AddTicker(FTickerDelegate::CreateLambda([this, kTickRate](float TimeDelta)
+	{
+		// ticker passes in frame-delta, not tick delta...
+		InnerTick(kTickRate);
+		return true;
+	}),
+	kTickRate);
+}
+
+void FGauntletModuleImpl::ShutdownModule()
+{
+	FCoreUObjectDelegates::PreLoadMap.RemoveAll(this);
+	FCoreUObjectDelegates::PostLoadMapWithWorld.RemoveAll(this);
+
+	if (TickHandle.IsValid())
+	{
+		FTicker::GetCoreTicker().RemoveTicker(TickHandle);
+	}
+
+	UE_LOG(LogGauntlet, Display, TEXT("Gauntlet Shutdown"));
+
+}
+
+void FGauntletModuleImpl::LoadControllers()
+{
+	FString ControllerString;
+
+	FParse::Value(FCommandLine::Get(), TEXT("gauntlet="), ControllerString, false);
+
+	if (ControllerString.Len())
+	{
+		TArray<FString> ControllerNames;
+
+		ControllerString.ParseIntoArrayWS(ControllerNames, TEXT(","));
+
+		TSet<FString> ControllersToCreate;
+
+		for (const FString& Name : ControllerNames)
+		{
+			UClass* TestClass = nullptr;
+
+			for (TObjectIterator<UClass> It; It; ++It)
+			{
+				UClass* Class = *It;
+
+				FString FullName = Name;
+
+				// Search for SomethingTestController and ControllerSomethingTest. The last is legacy
+				FString PartialName1 = Name + TEXT("Controller");
+				FString PartialName2 = FString(TEXT("Controller")) + Name;
+
+				if (Class->IsChildOf<UGauntletTestController>())
+				{
+					FString ClassName = Class->GetName();
+					if (ClassName == FullName || ClassName.EndsWith(PartialName1) || ClassName.EndsWith(PartialName2))
+					{
+						// Gauntlet has a couple of test classes, so need to differentiate between "GauntletFooTest" and "GameFooTest".
+						// that will both be launched via -gauntlet=FooTest
+						bool GauntletDefault = ClassName.StartsWith(TEXT("Gauntlet"));
+
+						TestClass = Class;
+
+						// If not gauntlet stop searching
+						if (!GauntletDefault)
+						{
+							break;
+						}
+					}
+				}
+			}			
+
+			checkf(TestClass, TEXT("Could not find class for controller %s"), *Name);
+
+			UGauntletTestController* NewController = NewObject<UGauntletTestController>(GetTransientPackage(), TestClass);
+
+			check(NewController);
+
+			UE_LOG(LogGauntlet, Display, TEXT("Added Gauntlet controller %s"), *Name);
+
+			// Important - add the controller first! Some controllers may trigger GC's which would
+			// then result in them being collected...
+			Controllers.Add(NewController);
+		}
+	}
+
+	for (auto Controller : Controllers)
+	{
+		Controller->OnInit();
+	}
+}
+
+
+void FGauntletModuleImpl::SetGameStateToTestStateMapping(const TMap<UClass*, FName>& Mapping)
+{
+	GameStateStateMap  = Mapping;
+}
+
+void FGauntletModuleImpl::SetWorldToTestStateMapping(const TMap<FString, FName>& Mapping)
+{
+	MapGameStateMap = Mapping;
+}
+
+void FGauntletModuleImpl::AddReferencedObjects(FReferenceCollector& Collector)
+{
+	for (auto Controller : Controllers)
+	{
+		Collector.AddReferencedObject(Controller);
+	}
+}
+
+void FGauntletModuleImpl::BroadcastStateChange(FName NewState)
+{
+	if (NewState == CurrentState && StateSet)
+	{
+		return;
+	}
+
+	for (auto Controller : Controllers)
+	{
+		Controller->OnStateChange(CurrentState, NewState);
+	}
+
+	CurrentState = NewState;
+	StateSet = true;
+	TimeInCurrentState = 0.0;
+
+	UE_LOG(LogGauntlet, Display, TEXT("Changed state to %s"), *NewState.ToString());
+}
+
+void FGauntletModuleImpl::InnerPreMapChange(const FString& MapName)
+{
+	for (auto Controller : Controllers)
+	{
+		Controller->OnPreMapChange();
+	}
+}
+
+void FGauntletModuleImpl::InnerPostMapChange(UWorld* World)
+{
+	CurrentMap = World->GetMapName();
+
+	for (auto Controller : Controllers)
+	{
+		Controller->OnPostMapChange(World);
+	}
+
+	// process mappings
+	if (MapGameStateMap.Contains(CurrentMap))
+	{
+		BroadcastStateChange(MapGameStateMap[CurrentMap]);
+	}
+}
+
+// FTickableGameObject interface
+void FGauntletModuleImpl::InnerTick(const float TimeDelta)
+{
+    QUICK_SCOPE_CYCLE_COUNTER(STAT_FGauntletModuleImpl_InnerTick);
+	
+    if (AGameStateBase* GameState = GWorld->GetGameState())
+	{
+		if (GameState->GetClass() != CurrentGameStateClass)
+		{
+			CurrentGameStateClass = GameState->GetClass();
+
+			FName NextState;
+
+			for (const auto& KP : GameStateStateMap)
+			{
+				if (CurrentGameStateClass->IsChildOf(KP.Key))
+				{
+					NextState = KP.Value;
+
+					if (CurrentGameStateClass == KP.Key)
+					{
+						break;
+					}
+				}
+			}
+
+			if (!NextState.IsNone())
+			{
+				BroadcastStateChange(NextState);
+			}
+		}
+
+		if (ScreenshotPeriod > 0.0f && (FPlatformTime::Seconds() - LastScreenshotTime) > ScreenshotPeriod)
+		{
+			TakeScreenshot();
+			LastScreenshotTime = FPlatformTime::Seconds();
+		}
+	}
+
+	TimeInCurrentState += TimeDelta;
+
+	for (auto Controller : Controllers)
+	{
+		Controller->OnTick(TimeDelta);
+	}
+}
+
+void FGauntletModuleImpl::TakeScreenshot()
+{
+	FScreenshotRequest::RequestScreenshot(true);
+}
+
+
+UGauntletTestController* FGauntletModuleImpl::GetTestController(UClass* ControllerClass)
+{
+	for (auto Controller : Controllers)
+	{
+		if (Controller->IsA(ControllerClass))
+		{
+			return Controller;
+		}
+	}
+
+	return nullptr;
+}
diff --git a/Engine/Plugins/Experimental/Gauntlet/Source/Gauntlet/Private/GauntletTestController.cpp b/Engine/Plugins/Experimental/Gauntlet/Source/Gauntlet/Private/GauntletTestController.cpp
new file mode 100644
--- /dev/null
+++ b/Engine/Plugins/Experimental/Gauntlet/Source/Gauntlet/Private/GauntletTestController.cpp
@@ -0,0 +1,52 @@
+// Copyright 1998-2018 Epic Games, Inc. All Rights Reserved.
+
+#include "GauntletTestController.h"
+#include "Engine/Engine.h"
+
+
+
+UGauntletTestController::UGauntletTestController(const FObjectInitializer& ObjectInitializer)
+{
+	if (!HasAnyFlags(RF_ClassDefaultObject))
+	{
+		ParentModule = &FModuleManager::Get().GetModuleChecked<FGauntletModule>(TEXT("Gauntlet"));
+	}
+}
+
+UGauntletTestController::~UGauntletTestController()
+{
+
+}
+
+FName UGauntletTestController::GetCurrentState() const
+{
+	check(ParentModule);
+	return ParentModule->GetCurrentState();
+}
+
+double UGauntletTestController::GetTimeInCurrentState() const
+{
+	check(ParentModule);
+	return ParentModule->GetTimeInCurrentState();
+}
+
+FString UGauntletTestController::GetCurrentMap() const
+{
+	return GWorld->GetName();
+}
+
+/**
+* Returns the gauntlet module running this test
+*/
+FGauntletModule* UGauntletTestController::GetGauntlet()
+{
+	return ParentModule;
+}
+
+void UGauntletTestController::EndTest(int32 ExitCode /*= 0*/)
+{
+	UE_LOG(LogGauntlet, Display, TEXT("**** TEST COMPLETE. EXIT CODE: %d ****"), ExitCode);
+	// we flush logs because we don't (currently...) want to treat shutdown errors as failures
+	GLog->PanicFlushThreadedLogs();
+	FPlatformMisc::RequestExit(1);
+}
diff --git a/Engine/Plugins/Experimental/Gauntlet/Source/Gauntlet/Private/GauntletTestControllerBootTest.cpp b/Engine/Plugins/Experimental/Gauntlet/Source/Gauntlet/Private/GauntletTestControllerBootTest.cpp
new file mode 100644
--- /dev/null
+++ b/Engine/Plugins/Experimental/Gauntlet/Source/Gauntlet/Private/GauntletTestControllerBootTest.cpp
@@ -0,0 +1,21 @@
+// Copyright 1998-2018 Epic Games, Inc. All Rights Reserved.
+
+#include "GauntletTestControllerBootTest.h"
+
+
+
+void UGauntletTestControllerBootTest::OnTick(float TimeDelta)
+{
+	if (IsBootProcessComplete())
+	{
+		EndTest(0);
+	}
+	else
+	{
+		if (GetTimeInCurrentState() > 300)
+		{
+			UE_LOG(LogGauntlet, Error, TEXT("Failing boot test after 300 secs!"));
+			EndTest(-1);
+		}
+	}
+}
\ No newline at end of file
diff --git a/Engine/Plugins/Experimental/Gauntlet/Source/Gauntlet/Private/GauntletTestControllerErrorTest.cpp b/Engine/Plugins/Experimental/Gauntlet/Source/Gauntlet/Private/GauntletTestControllerErrorTest.cpp
new file mode 100644
--- /dev/null
+++ b/Engine/Plugins/Experimental/Gauntlet/Source/Gauntlet/Private/GauntletTestControllerErrorTest.cpp
@@ -0,0 +1,52 @@
+// Copyright 1998-2018 Epic Games, Inc. All Rights Reserved.
+
+#include "GauntletTestControllerErrorTest.h"
+#include "Engine/Engine.h"
+#include "Misc/CommandLine.h"
+
+
+void UGauntletTestControllerErrorTest::OnInit()
+{
+	ErrorDelay = 0;
+	ErrorType = TEXT("check");
+
+	FParse::Value(FCommandLine::Get(), TEXT("errortest.delay="), ErrorDelay);
+	FParse::Value(FCommandLine::Get(), TEXT("errortest.type="), ErrorType);
+}
+
+void UGauntletTestControllerErrorTest::OnTick(float TimeDelta)
+{
+	if (GetTimeInCurrentState() > ErrorDelay)
+	{
+		if (ErrorType == TEXT("ensure"))
+		{
+			UE_LOG(LogGauntlet, Display, TEXT("Issuing ensure as requested"));
+			ensureMsgf(false, TEXT("Ensuring false...."));
+			EndTest(-1);
+		}
+		else if (ErrorType == TEXT("check"))
+		{
+			UE_LOG(LogGauntlet, Display, TEXT("Issuing failed check as requested"));
+			checkf(false, TEXT("Asserting as requested"));
+		}		
+		else if (ErrorType == TEXT("fatal"))
+		{
+			UE_LOG(LogGauntlet, Fatal, TEXT("Issuing fatal error as requested"));
+		}
+		else if (ErrorType == TEXT("gpf"))
+		{
+#ifndef PVS_STUDIO
+			UE_LOG(LogGauntlet, Display, TEXT("Issuing GPF as requested"));
+			int* Ptr = (int*)0;
+			CA_SUPPRESS(6011);
+			*Ptr = 42; // -V522
+#endif // PVS_STUDIO
+		}
+		else
+		{
+			UE_LOG(LogGauntlet, Error, TEXT("No recognized error request. Failing test"));
+			EndTest(-1);
+		}
+	}
+}
+
diff --git a/Engine/Plugins/Experimental/Gauntlet/Source/Gauntlet/Public/GauntletModule.h b/Engine/Plugins/Experimental/Gauntlet/Source/Gauntlet/Public/GauntletModule.h
new file mode 100644
--- /dev/null
+++ b/Engine/Plugins/Experimental/Gauntlet/Source/Gauntlet/Public/GauntletModule.h
@@ -0,0 +1,107 @@
+// Copyright 1998-2018 Epic Games, Inc. All Rights Reserved.
+
+#pragma once
+
+#include "Modules/ModuleInterface.h"
+#include "Modules/ModuleManager.h"
+#include "UObject/GCObject.h"
+#include "Tickable.h"
+#include "Containers/UnrealString.h"
+
+/**
+*	Gauntlet states. Define your own states by inheriting from this or some other
+*	form that you see fit.
+*
+*/
+struct GAUNTLET_API FGauntletStates
+{
+	static FName Initialized;
+};
+
+class UGauntletTestController;
+
+/**
+*	Main Gauntlet Module. This is module is responsible for managing its underlying controllers and
+*	propagating events and state changes to them as necessary.
+*	
+*	After initialization you should be prepared to feed this module with states that your controllers
+*	can respond to. This can be done either manually (BroadcastStateChange), through one of the helpers
+*	(e.g. SetGameStateToTestStateMapping), or a combination of both.
+*/
+class GAUNTLET_API FGauntletModule : public IModuleInterface
+{
+public:
+
+	/**
+	*
+	*	This is a convenient way of binding a list of AGameState  types to your own state defines. When
+	*	there is a state change in the world (e.g. from loading a new map) the mapped state type will be
+	*	broadcast to all controllers.
+	*
+	* @param: 	Mapping - A map e.g. [AGameStateSomethingClass, StateSomething]
+	* @return: void
+	*/
+	virtual void			SetGameStateToTestStateMapping(const TMap<UClass*, FName>& Mapping) = 0;
+
+	/**
+	*
+	*	This is a convenient way of binding a list of maps types to your own state defines. When
+	*	there is a state change in the world (e.g. from loading a new map) the mapped state type will be
+	*	broadcast to all controllers.
+	*
+	* @param: 	Mapping - A map e.g. [TEXT("FrontendMap"), StateFrontend]
+	* @return: void
+	*/
+	virtual void			SetWorldToTestStateMapping(const TMap<FString, FName>& Mapping) = 0;
+
+	/**
+	*	Manually Broadcasts a state change to all current Gauntlet controllers. This can be used as an alternative
+	*	or addition to the functions above for broadcasting state changes to running controllers.
+	*
+	* @param: 	FName		New State
+	* @return: void
+	*/
+	virtual void			BroadcastStateChange(FName NewState) = 0;
+
+	/**
+	* Returns the current  state
+	*
+	* @return: FName
+	*/
+	virtual FName			GetCurrentState() const = 0;
+
+	/**
+	* Returns the time spent in the current state
+	*
+	* @return: double
+	*/
+	virtual double			GetTimeInCurrentState() const = 0;
+
+	/**
+	 * Sets the rate for screenshots to be taken (default = 0)
+	 */
+	virtual void			SetScreenshotPeriod(float Period) = 0;
+
+	/**
+	 * Returns the first controller (if any) matching the provided name
+	 */
+	virtual UGauntletTestController*	GetTestController(UClass* ControllerClass) = 0;
+
+	/**
+	* Templated casty version of the above
+	*/
+	template< class T >
+	inline T* GetTestController()
+	{
+		return static_cast<T*>(GetTestController(T::StaticClass()));
+	}
+};
+
+
+
+/** Logging category */
+GAUNTLET_API DECLARE_LOG_CATEGORY_EXTERN(LogGauntlet, Log, All);
+
+
+
+
diff --git a/Engine/Plugins/Experimental/Gauntlet/Source/Gauntlet/Public/GauntletTestController.h b/Engine/Plugins/Experimental/Gauntlet/Source/Gauntlet/Public/GauntletTestController.h
new file mode 100644
--- /dev/null
+++ b/Engine/Plugins/Experimental/Gauntlet/Source/Gauntlet/Public/GauntletTestController.h
@@ -0,0 +1,94 @@
+// Copyright 1998-2018 Epic Games, Inc. All Rights Reserved.
+
+#pragma once
+
+#include "GauntletModule.h"
+#include "GauntletTestController.generated.h"
+
+
+/**
+ *	Base class for games to implement test controllers that use the Gauntlet native
+ *	framework. This is a very thin class that is created automatically based on 
+ *	command line params (-gauntlet=MyControllerName) and provides easily overridden
+ *	functions that represent state changes and ticking
+ *
+ *	In essence your derived class should implement logic that starts and monitors
+ *	a test, then calls EndTest(Result) when the desired criteria are met (or not!)
+ */
+UCLASS()
+class GAUNTLET_API UGauntletTestController : public UObject
+{
+	GENERATED_BODY()
+
+	friend class FGauntletModule;
+
+public:
+
+	/** Default constructur */
+	UGauntletTestController(const FObjectInitializer& ObjectInitializer);
+	virtual ~UGauntletTestController();
+
+	// Overridable delegates for some of the most useful test points
+
+	/**
+	 *	Called when the controller is first initialized
+	 */
+	virtual void	OnInit() {}
+
+	/**
+	 *	Called prior to a map change
+	 */
+	virtual void	OnPreMapChange() {}
+
+	/**
+	 *	Called after a map change. GetCurrentMap() will now return the new map
+	 */
+	virtual void	OnPostMapChange(UWorld* World) {}
+
+	/**
+	 *	Called periodically to let the controller check and control state
+	 */
+	virtual void	OnTick(float TimeDelta) {}
+
+	/**
+	 *	Called when a state change is applied to the module. States are game-driven.
+	 *	GetCurrentState() == OldState until this function returns
+	 */
+	virtual void	OnStateChange(FName OldState, FName NewState) {}
+
+protected:
+
+	/**
+	 *	Returns the current state applied to Gauntlet
+	 */
+	FName			GetCurrentState() const;
+
+	/**
+	 *	Return the time since OnStateChange was called with the current state
+	 */
+	double			GetTimeInCurrentState() const;
+
+	/**
+	 *	Return the name of the current persistent map
+	 */
+	FString			GetCurrentMap() const;
+
+	/**
+	 *	Called to end the current test with the provided exit code. This will cause the
+	 *	app to exit with the provided return code.
+	 */
+	void			EndTest(int32 ExitCode = 0);
+
+	/**
+	 * Returns the gauntlet module running this test
+	 */
+	FGauntletModule* GetGauntlet();
+
+private:
+
+	FGauntletModule*  ParentModule;
+};
+
+
+
+
diff --git a/Engine/Plugins/Experimental/Gauntlet/Source/Gauntlet/Public/GauntletTestControllerBootTest.h b/Engine/Plugins/Experimental/Gauntlet/Source/Gauntlet/Public/GauntletTestControllerBootTest.h
new file mode 100644
--- /dev/null
+++ b/Engine/Plugins/Experimental/Gauntlet/Source/Gauntlet/Public/GauntletTestControllerBootTest.h
@@ -0,0 +1,18 @@
+// Copyright 1998-2018 Epic Games, Inc. All Rights Reserved.
+
+#pragma once
+
+#include "GauntletTestController.h"
+#include "GauntletTestControllerBootTest.generated.h"
+
+UCLASS()
+class GAUNTLET_API UGauntletTestControllerBootTest : public UGauntletTestController
+{
+	GENERATED_BODY()
+
+protected:
+
+	virtual bool	IsBootProcessComplete()		const { return false; }
+	void			OnTick(float TimeDelta)		override;
+};
+
diff --git a/Engine/Plugins/Experimental/Gauntlet/Source/Gauntlet/Public/GauntletTestControllerErrorTest.h b/Engine/Plugins/Experimental/Gauntlet/Source/Gauntlet/Public/GauntletTestControllerErrorTest.h
new file mode 100644
--- /dev/null
+++ b/Engine/Plugins/Experimental/Gauntlet/Source/Gauntlet/Public/GauntletTestControllerErrorTest.h
@@ -0,0 +1,22 @@
+// Copyright 1998-2018 Epic Games, Inc. All Rights Reserved.
+
+#pragma once
+
+#include "GauntletTestController.h"
+#include "GauntletTestControllerErrorTest.generated.h"
+
+UCLASS()
+class GAUNTLET_API UGauntletTestControllerErrorTest : public UGauntletTestController
+{
+	GENERATED_BODY()
+
+protected:
+
+	float			ErrorDelay;
+	FString			ErrorType;
+	bool			RunOnServer;
+
+	void	OnInit() override;
+	void	OnTick(float TimeDelta)		override;
+};
+
