Make gauntlet tick time configurable on the commandline using -gauntlet.tickrate=

diff --git a/Engine/Plugins/Experimental/Gauntlet/Source/Gauntlet/Private/GauntletModule.cpp b/Engine/Plugins/Experimental/Gauntlet/Source/Gauntlet/Private/GauntletModule.cpp
--- a/Engine/Plugins/Experimental/Gauntlet/Source/Gauntlet/Private/GauntletModule.cpp
+++ b/Engine/Plugins/Experimental/Gauntlet/Source/Gauntlet/Private/GauntletModule.cpp
@@ -119,7 +119,9 @@ void FGauntletModuleImpl::OnPostEngineInit()
 
 	LoadControllers();
 
-	const float kTickRate = 1.0f;
+	float kTickRate = 1.0f;
+	FParse::Value(FCommandLine::Get(), TEXT("gauntlet.tickrate="), kTickRate);
+
 
 	FTicker::GetCoreTicker().AddTicker(FTickerDelegate::CreateLambda([this, kTickRate](float TimeDelta)
 	{
