Added helpers to GauntletTestController to get World / Player controller

diff --git a/Engine/Plugins/Experimental/Gauntlet/Source/Gauntlet/Private/GauntletTestController.cpp b/Engine/Plugins/Experimental/Gauntlet/Source/Gauntlet/Private/GauntletTestController.cpp
--- a/Engine/Plugins/Experimental/Gauntlet/Source/Gauntlet/Private/GauntletTestController.cpp
+++ b/Engine/Plugins/Experimental/Gauntlet/Source/Gauntlet/Private/GauntletTestController.cpp
@@ -2,7 +2,7 @@
 
 #include "GauntletTestController.h"
 #include "Engine/Engine.h"
-
+#include "Engine/GameViewPortClient.h"
 
 
 UGauntletTestController::UGauntletTestController(const FObjectInitializer& ObjectInitializer)
@@ -43,6 +43,18 @@ FGauntletModule* UGauntletTestController::GetGauntlet()
 	return ParentModule;
 }
 
+UWorld* UGauntletTestController::GetWorld() const
+{
+	UGameInstance* GameInstance = GEngine->GameViewport->GetGameInstance();
+
+	return GameInstance ? GameInstance->GetWorld() : nullptr;
+}
+
+APlayerController* UGauntletTestController::GetFirstPlayerController() const
+{
+	return GetWorld() ? GetWorld()->GetFirstPlayerController() : nullptr;	
+}
+
 void UGauntletTestController::EndTest(int32 ExitCode /*= 0*/)
 {
 	UE_LOG(LogGauntlet, Display, TEXT("**** TEST COMPLETE. EXIT CODE: %d ****"), ExitCode);
diff --git a/Engine/Plugins/Experimental/Gauntlet/Source/Gauntlet/Public/GauntletTestController.h b/Engine/Plugins/Experimental/Gauntlet/Source/Gauntlet/Public/GauntletTestController.h
--- a/Engine/Plugins/Experimental/Gauntlet/Source/Gauntlet/Public/GauntletTestController.h
+++ b/Engine/Plugins/Experimental/Gauntlet/Source/Gauntlet/Public/GauntletTestController.h
@@ -5,6 +5,7 @@
 #include "GauntletModule.h"
 #include "GauntletTestController.generated.h"
 
+class APlayerController;
 
 /**
  *	Base class for games to implement test controllers that use the Gauntlet native
@@ -56,6 +57,16 @@ public:
 	 */
 	virtual void	OnStateChange(FName OldState, FName NewState) {}
 
+	/**
+	 * Return the current world
+	 */
+	UWorld* GetWorld() const override;
+
+	/**
+	 *	Helper function that returns the first player controller in the world (may be null depending on when called).
+	 */
+	APlayerController* GetFirstPlayerController() const;
+
 protected:
 
 	/**
