Change include case to fix Linux CIS error.

diff --git a/Engine/Plugins/Experimental/Gauntlet/Source/Gauntlet/Private/GauntletTestController.cpp b/Engine/Plugins/Experimental/Gauntlet/Source/Gauntlet/Private/GauntletTestController.cpp
--- a/Engine/Plugins/Experimental/Gauntlet/Source/Gauntlet/Private/GauntletTestController.cpp
+++ b/Engine/Plugins/Experimental/Gauntlet/Source/Gauntlet/Private/GauntletTestController.cpp
@@ -2,7 +2,7 @@
 
 #include "GauntletTestController.h"
 #include "Engine/Engine.h"
-#include "Engine/GameViewPortClient.h"
+#include "Engine/GameViewportClient.h"
 
 
 UGauntletTestController::UGauntletTestController(const FObjectInitializer& ObjectInitializer)
