[IslandInspector] Inspector logic can be instrumented on server
+ Additional changes herein provide a mechanism (http) to transfer client-rendered screenshots to the server-instrumented inspector runtime
+ Client inspector logic has been encapsulated via the componentized PlayerController pattern introduced by BobT

diff --git a/Engine/Plugins/Experimental/Gauntlet/Source/Gauntlet/Private/GauntletTestController.cpp b/Engine/Plugins/Experimental/Gauntlet/Source/Gauntlet/Private/GauntletTestController.cpp
--- a/Engine/Plugins/Experimental/Gauntlet/Source/Gauntlet/Private/GauntletTestController.cpp
+++ b/Engine/Plugins/Experimental/Gauntlet/Source/Gauntlet/Private/GauntletTestController.cpp
@@ -45,9 +45,7 @@ FGauntletModule* UGauntletTestController::GetGauntlet()
 
 UWorld* UGauntletTestController::GetWorld() const
 {
-	UGameInstance* GameInstance = GEngine->GameViewport->GetGameInstance();
-
-	return GameInstance ? GameInstance->GetWorld() : nullptr;
+	return GWorld;
 }
 
 APlayerController* UGauntletTestController::GetFirstPlayerController() const
