Remove hardcoded timeout from BootTest test controller - this should be defined by the test node in csharp.

diff --git a/Engine/Plugins/Experimental/Gauntlet/Source/Gauntlet/Private/GauntletTestControllerBootTest.cpp b/Engine/Plugins/Experimental/Gauntlet/Source/Gauntlet/Private/GauntletTestControllerBootTest.cpp
--- a/Engine/Plugins/Experimental/Gauntlet/Source/Gauntlet/Private/GauntletTestControllerBootTest.cpp
+++ b/Engine/Plugins/Experimental/Gauntlet/Source/Gauntlet/Private/GauntletTestControllerBootTest.cpp
@@ -10,12 +10,4 @@ void UGauntletTestControllerBootTest::OnTick(float TimeDelta)
 	{
 		EndTest(0);
 	}
-	else
-	{
-		if (GetTimeInCurrentState() > 300)
-		{
-			UE_LOG(LogGauntlet, Error, TEXT("Failing boot test after 300 secs!"));
-			EndTest(-1);
-		}
-	}
 }
\ No newline at end of file
