Gauntlet: force exit only if platform doesn't support quit
#rnx

diff --git a/Engine/Plugins/Experimental/Gauntlet/Source/Gauntlet/Private/GauntletTestController.cpp b/Engine/Plugins/Experimental/Gauntlet/Source/Gauntlet/Private/GauntletTestController.cpp
--- a/Engine/Plugins/Experimental/Gauntlet/Source/Gauntlet/Private/GauntletTestController.cpp
+++ b/Engine/Plugins/Experimental/Gauntlet/Source/Gauntlet/Private/GauntletTestController.cpp
@@ -69,5 +69,6 @@ void UGauntletTestController::EndTest(int32 ExitCode /*= 0*/)
 	UE_LOG(LogGauntlet, Display, TEXT("**** TEST COMPLETE. EXIT CODE: %d ****"), ExitCode);
 	// we flush logs because we don't (currently...) want to treat shutdown errors as failures
 	GLog->PanicFlushThreadedLogs();
-	FPlatformMisc::RequestExit(1);
+	// force exit only if platform doesn't support quitting
+	FPlatformMisc::RequestExit(!FPlatformProperties::SupportsQuit());
 }
