Downgraded shutdown message in Gauntlet from Display to Log

diff --git a/Engine/Plugins/Experimental/Gauntlet/Source/Gauntlet/Private/GauntletModule.cpp b/Engine/Plugins/Experimental/Gauntlet/Source/Gauntlet/Private/GauntletModule.cpp
--- a/Engine/Plugins/Experimental/Gauntlet/Source/Gauntlet/Private/GauntletModule.cpp
+++ b/Engine/Plugins/Experimental/Gauntlet/Source/Gauntlet/Private/GauntletModule.cpp
@@ -165,7 +165,7 @@ void FGauntletModuleImpl::ShutdownModule()
 		FTicker::GetCoreTicker().RemoveTicker(TickHandle);
 	}
 
-	UE_LOG(LogGauntlet, Display, TEXT("Gauntlet Shutdown"));
+	UE_LOG(LogGauntlet, Log, TEXT("Gauntlet Shutdown"));
 
 }
 
