Enable Gauntlet support for PIE, particularly with CitySample.

diff --git a/Engine/Plugins/Experimental/Gauntlet/Source/Gauntlet/Gauntlet.Build.cs b/Engine/Plugins/Experimental/Gauntlet/Source/Gauntlet/Gauntlet.Build.cs
--- a/Engine/Plugins/Experimental/Gauntlet/Source/Gauntlet/Gauntlet.Build.cs
+++ b/Engine/Plugins/Experimental/Gauntlet/Source/Gauntlet/Gauntlet.Build.cs
@@ -10,9 +10,17 @@ namespace UnrealBuildTool.Rules
 				new string[] {
 					"Core",
 					"CoreUObject",
-					"Engine",
+					"Engine"
 				});
 
+			if (Target.bCompileAgainstEditor)
+			{
+				PublicDependencyModuleNames.AddRange(new string[]
+				{
+				"UnrealEd"
+				});
+			}
+
 			PublicIncludePaths.AddRange(
 				new string[]
 				{
diff --git a/Engine/Plugins/Experimental/Gauntlet/Source/Gauntlet/Private/GauntletModule.cpp b/Engine/Plugins/Experimental/Gauntlet/Source/Gauntlet/Private/GauntletModule.cpp
--- a/Engine/Plugins/Experimental/Gauntlet/Source/Gauntlet/Private/GauntletModule.cpp
+++ b/Engine/Plugins/Experimental/Gauntlet/Source/Gauntlet/Private/GauntletModule.cpp
@@ -9,7 +9,9 @@
 #include "Containers/Ticker.h"
 #include "GameFramework/GameStateBase.h"
 #include "UnrealClient.h"
-
+#if WITH_EDITOR
+#include "Editor.h"
+#endif
 
 
 class GAUNTLET_API FGauntletModuleImpl
@@ -129,6 +131,12 @@ void FGauntletModuleImpl::StartupModule()
 		UE_LOG(LogGauntlet, Display, TEXT("Gauntlet Initialized"));
 		FCoreDelegates::OnPostEngineInit.AddRaw(this, &FGauntletModuleImpl::OnPostEngineInit);	
 	}
+#if WITH_EDITOR
+	else
+	{
+		FEditorDelegates::PreBeginPIE.AddLambda([this](bool bIsSimulating){ this->OnPostEngineInit(); });
+	}
+#endif
 }
 
 void FGauntletModuleImpl::OnPostEngineInit()
@@ -316,7 +324,18 @@ void FGauntletModuleImpl::InnerTick(const float TimeDelta)
 {
     QUICK_SCOPE_CYCLE_COUNTER(STAT_FGauntletModuleImpl_InnerTick);
 	
-    if (AGameStateBase* GameState = GWorld->GetGameState())
+	UWorld* World = GWorld;
+
+#if WITH_EDITOR
+	UEditorEngine* EditorEngine = Cast<UEditorEngine>(GEngine);
+	if (GIsEditor && EditorEngine != nullptr && EditorEngine->PlayWorld != nullptr)
+	{
+		// lets use PlayWorld during PIE/Simulate
+		World = EditorEngine->PlayWorld;
+	}
+#endif
+
+    if (AGameStateBase* GameState = World->GetGameState())
 	{
 		if (GameState->GetClass() != CurrentGameStateClass)
 		{
diff --git a/Engine/Plugins/Experimental/Gauntlet/Source/Gauntlet/Private/GauntletTestController.cpp b/Engine/Plugins/Experimental/Gauntlet/Source/Gauntlet/Private/GauntletTestController.cpp
--- a/Engine/Plugins/Experimental/Gauntlet/Source/Gauntlet/Private/GauntletTestController.cpp
+++ b/Engine/Plugins/Experimental/Gauntlet/Source/Gauntlet/Private/GauntletTestController.cpp
@@ -4,6 +4,9 @@
 #include "Engine/World.h"
 #include "GauntletModule.h"
 #include "Modules/ModuleManager.h"
+#if WITH_EDITOR
+#include "Editor/EditorEngine.h"
+#endif
 
 #include UE_INLINE_GENERATED_CPP_BY_NAME(GauntletTestController)
 
@@ -35,7 +38,7 @@ double UGauntletTestController::GetTimeInCurrentState() const
 
 FString UGauntletTestController::GetCurrentMap() const
 {
-	return GWorld->GetName();
+	return GetWorld()->GetName();
 }
 
 /**
@@ -59,6 +62,14 @@ void UGauntletTestController::MarkHeartbeatActive(const FString& OptionalStatusM
 
 UWorld* UGauntletTestController::GetWorld() const
 {
+#if WITH_EDITOR
+	UEditorEngine* EditorEngine = Cast<UEditorEngine>(GEngine);
+	if (GIsEditor && EditorEngine != nullptr && EditorEngine->PlayWorld != nullptr)
+	{
+		// lets use PlayWorld during PIE/Simulate
+		return EditorEngine->PlayWorld;
+	}
+#endif
 	return GWorld;
 }
 
@@ -72,6 +83,13 @@ void UGauntletTestController::EndTest(int32 ExitCode /*= 0*/)
 	UE_LOG(LogGauntlet, Display, TEXT("**** TEST COMPLETE. EXIT CODE: %d ****"), ExitCode);
 	// we flush logs because we don't (currently...) want to treat shutdown errors as failures
 	GLog->FlushThreadedLogs();
+#if WITH_EDITOR
+	if (GIsEditor)
+	{
+		GEngine->DeferredCommands.Add(TEXT("QUIT_EDITOR"));
+		return;
+	}
+#endif
 	// force exit only if platform doesn't support quitting
 	FPlatformMisc::RequestExitWithStatus(!FPlatformProperties::SupportsQuit(), static_cast<uint8>(ExitCode));
 }
