Enabling GauntletModule to support multiple entries/exits from PIE without re-creating its controllers.

diff --git a/Engine/Plugins/Experimental/Gauntlet/Source/Gauntlet/Private/GauntletModule.cpp b/Engine/Plugins/Experimental/Gauntlet/Source/Gauntlet/Private/GauntletModule.cpp
--- a/Engine/Plugins/Experimental/Gauntlet/Source/Gauntlet/Private/GauntletModule.cpp
+++ b/Engine/Plugins/Experimental/Gauntlet/Source/Gauntlet/Private/GauntletModule.cpp
@@ -61,7 +61,7 @@ protected:
 	void		InnerTick(const float TimeDelta);
 
 	/** Performs actual initialization */
-	void		OnPostEngineInit();
+	void		PerformInitialization();
 
 	/** Loads all controllers based on command line arts */
 	void		LoadControllers();
@@ -129,37 +129,40 @@ void FGauntletModuleImpl::StartupModule()
 	if (IsRunningGame() || IsRunningDedicatedServer())
 	{
 		UE_LOG(LogGauntlet, Display, TEXT("Gauntlet Initialized"));
-		FCoreDelegates::OnPostEngineInit.AddRaw(this, &FGauntletModuleImpl::OnPostEngineInit);	
+		FCoreDelegates::OnPostEngineInit.AddRaw(this, &FGauntletModuleImpl::PerformInitialization);
 	}
 #if WITH_EDITOR
 	else
 	{
-		FEditorDelegates::PreBeginPIE.AddLambda([this](bool bIsSimulating){ this->OnPostEngineInit(); });
+		FEditorDelegates::PreBeginPIE.AddLambda([this](bool bIsSimulating){ this->PerformInitialization(); });
 	}
 #endif
 }
 
-void FGauntletModuleImpl::OnPostEngineInit()
+void FGauntletModuleImpl::PerformInitialization()
 {
-	FCoreUObjectDelegates::PostLoadMapWithWorld.AddRaw(this, &FGauntletModuleImpl::InnerPostMapChange);
-	FCoreUObjectDelegates::PreLoadMap.AddRaw(this, &FGauntletModuleImpl::InnerPreMapChange);
+	if (Controllers.Num() == 0)
+	{
+		FCoreUObjectDelegates::PostLoadMapWithWorld.AddRaw(this, &FGauntletModuleImpl::InnerPostMapChange);
+		FCoreUObjectDelegates::PreLoadMap.AddRaw(this, &FGauntletModuleImpl::InnerPreMapChange);
 
-	FParse::Value(FCommandLine::Get(), TEXT("gauntlet.screenshotperiod="), ScreenshotPeriod);
-	FParse::Value(FCommandLine::Get(), TEXT("gauntlet.heartbeatperiod="), HeartbeatPeriod);
+		FParse::Value(FCommandLine::Get(), TEXT("gauntlet.screenshotperiod="), ScreenshotPeriod);
+		FParse::Value(FCommandLine::Get(), TEXT("gauntlet.heartbeatperiod="), HeartbeatPeriod);
 
-	LoadControllers();
+		LoadControllers();
 
-	float kTickRate = 1.0f;
-	FParse::Value(FCommandLine::Get(), TEXT("gauntlet.tickrate="), kTickRate);
+		float kTickRate = 1.0f;
+		FParse::Value(FCommandLine::Get(), TEXT("gauntlet.tickrate="), kTickRate);
 
 
-	TickHandle = FTSTicker::GetCoreTicker().AddTicker(FTickerDelegate::CreateLambda([this, kTickRate](float TimeDelta)
-	{
-		// ticker passes in frame-delta, not tick delta...
-		InnerTick(kTickRate);
-		return true;
-	}),
-	kTickRate);
+		TickHandle = FTSTicker::GetCoreTicker().AddTicker(FTickerDelegate::CreateLambda([this, kTickRate](float TimeDelta)
+		{
+			// ticker passes in frame-delta, not tick delta...
+			InnerTick(kTickRate);
+			return true;
+		}),
+		kTickRate);
+	}
 }
 
 void FGauntletModuleImpl::ShutdownModule()
