Fix odd spacing in UGauntletTestControllerBootTest and missing virtual keyword on the OnTick function.

#rnx

diff --git a/Engine/Plugins/Experimental/Gauntlet/Source/Gauntlet/Public/GauntletTestControllerBootTest.h b/Engine/Plugins/Experimental/Gauntlet/Source/Gauntlet/Public/GauntletTestControllerBootTest.h
--- a/Engine/Plugins/Experimental/Gauntlet/Source/Gauntlet/Public/GauntletTestControllerBootTest.h
+++ b/Engine/Plugins/Experimental/Gauntlet/Source/Gauntlet/Public/GauntletTestControllerBootTest.h
@@ -12,7 +12,7 @@ class GAUNTLET_API UGauntletTestControllerBootTest : public UGauntletTestControl
 
 protected:
 
-	virtual bool	IsBootProcessComplete()		const { return false; }
-	void			OnTick(float TimeDelta)		override;
+	virtual bool IsBootProcessComplete() const { return false; }
+	virtual void OnTick(float TimeDelta) override;
 };
 
