UE: Skip ticking gauntlet if ticked while there is no world.

diff --git a/Engine/Plugins/Experimental/Gauntlet/Source/Gauntlet/Private/GauntletModule.cpp b/Engine/Plugins/Experimental/Gauntlet/Source/Gauntlet/Private/GauntletModule.cpp
--- a/Engine/Plugins/Experimental/Gauntlet/Source/Gauntlet/Private/GauntletModule.cpp
+++ b/Engine/Plugins/Experimental/Gauntlet/Source/Gauntlet/Private/GauntletModule.cpp
@@ -338,6 +338,11 @@ void FGauntletModuleImpl::InnerTick(const float TimeDelta)
 	}
 #endif
 
+    if (!World)
+    {
+        return;
+    }
+
     if (AGameStateBase* GameState = World->GetGameState())
 	{
 		if (GameState->GetClass() != CurrentGameStateClass)
